TaskGraphVisualizer - Works only if WITH_EDITOR

diff --git a/Engine/Source/Developer/TaskGraph/Private/STaskGraph.cpp b/Engine/Source/Developer/TaskGraph/Private/STaskGraph.cpp
--- a/Engine/Source/Developer/TaskGraph/Private/STaskGraph.cpp
+++ b/Engine/Source/Developer/TaskGraph/Private/STaskGraph.cpp
@@ -172,7 +172,9 @@ public:
 	}
 	virtual void DisplayProfileVisualizer( TSharedPtr< FVisualizerEvent > InProfileData, const TCHAR* InProfilerType ) OVERRIDE
 	{
+#if	WITH_EDITOR
 		::DisplayProfileVisualizer( InProfileData, InProfilerType );
+#endif // WITH_EDITOR
 	}
 };
 IMPLEMENT_MODULE(FProfileVisualizerModule, TaskGraph);
