Visualizer tweaks, more elements can be displayed on screen now

diff --git a/Engine/Source/Developer/TaskGraph/Private/SBarVisualizer.cpp b/Engine/Source/Developer/TaskGraph/Private/SBarVisualizer.cpp
--- a/Engine/Source/Developer/TaskGraph/Private/SBarVisualizer.cpp
+++ b/Engine/Source/Developer/TaskGraph/Private/SBarVisualizer.cpp
@@ -116,7 +116,7 @@ void SBarVisualizer::Construct( const FArguments& InArgs )
 				// List of thread graphs
 				SAssignNew( BarGraphsList, SListView< TSharedPtr< FVisualizerEvent > > )
 				// List view items are this tall
-				.ItemHeight( 50 )
+				.ItemHeight( 24 )
 				// Tell the list view where to get its source data
 				.ListItemsSource( &ProfileDataView )
 				// When the list view needs to generate a widget for some data item, use this method
diff --git a/Engine/Source/Developer/TaskGraph/Private/SEventsTree.cpp b/Engine/Source/Developer/TaskGraph/Private/SEventsTree.cpp
--- a/Engine/Source/Developer/TaskGraph/Private/SEventsTree.cpp
+++ b/Engine/Source/Developer/TaskGraph/Private/SEventsTree.cpp
@@ -128,7 +128,7 @@ void SEventsTree::Construct( const FArguments& InArgs )
 	this->ChildSlot
 	[
 		SNew( SVerticalBox )
-		+SVerticalBox::Slot().AutoHeight().Padding( 2.0f, 0.0f, 2.0f, 4.0f )
+		+SVerticalBox::Slot().AutoHeight().Padding( 1.0f, 0.0f, 1.0f, 2.0f )
 		[
 			SNew( SSearchBox )
 			.ToolTipText( NSLOCTEXT("TaskGraph", "FilterSearchHint", "Type here to search events.") )
@@ -143,7 +143,7 @@ void SEventsTree::Construct( const FArguments& InArgs )
 				// List of all events for the selected thread
 				SAssignNew( EventsListView, STreeView< TSharedPtr< FVisualizerEvent > > )
 				// List view items are this tall
-				.ItemHeight( 16 )
+				.ItemHeight( 12 )
 				// Tell the list view where to get its source data
 				.TreeItemsSource( &SelectedEventsView )
 				// When the list view needs to generate a widget for some data item, use this method
@@ -161,6 +161,7 @@ void SEventsTree::Construct( const FArguments& InArgs )
 					.DefaultLabel( NSLOCTEXT("TaskGraph", "ColumnName", "Name") ) 
 					.SortMode( TAttribute< EColumnSortMode::Type >::Create( TAttribute< EColumnSortMode::Type >::FGetter::CreateSP( this, &SEventsTree::GetColumnSortMode, NAME_NameColumn ) ) )
 					.OnSort( FOnSortModeChanged::CreateSP( this, &SEventsTree::OnColumnSortModeChanged ) )
+					.FillWidth( 1.0f )
 					.MenuContent()
 					[
 						NameMenuBuilder.MakeWidget()
@@ -170,6 +171,7 @@ void SEventsTree::Construct( const FArguments& InArgs )
 					.DefaultLabel( this, &SEventsTree::GetDurationColumnTitle ) 
 					.SortMode( TAttribute< EColumnSortMode::Type >::Create( TAttribute< EColumnSortMode::Type >::FGetter::CreateSP( this, &SEventsTree::GetColumnSortMode, NAME_DurationColumn ) ) )
 					.OnSort( FOnSortModeChanged::CreateSP( this, &SEventsTree::OnColumnSortModeChanged ) )
+					.FixedWidth( 128.0f )
 					.MenuContent()
 					[
 						MenuBuilder.MakeWidget()
diff --git a/Engine/Source/Developer/TaskGraph/Private/SGraphBar.cpp b/Engine/Source/Developer/TaskGraph/Private/SGraphBar.cpp
--- a/Engine/Source/Developer/TaskGraph/Private/SGraphBar.cpp
+++ b/Engine/Source/Developer/TaskGraph/Private/SGraphBar.cpp
@@ -211,7 +211,7 @@ FReply SGraphBar::OnMouseMove( const FGeometry& MyGeometry, const FPointerEvent&
 
 FVector2D SGraphBar::ComputeDesiredSize( float ) const
 {
-	return FVector2D( 500.0f, 50.0f );
+	return FVector2D( 8.0f, 16.0f );
 }
 
 void SGraphBar::SetEvents( const FVisualizerEventsArray& InEvents, double InStartTime, double InTotalTime )
diff --git a/Engine/Source/Developer/TaskGraph/Private/STimeline.cpp b/Engine/Source/Developer/TaskGraph/Private/STimeline.cpp
--- a/Engine/Source/Developer/TaskGraph/Private/STimeline.cpp
+++ b/Engine/Source/Developer/TaskGraph/Private/STimeline.cpp
@@ -179,5 +179,5 @@ FReply STimeline::OnMouseMove( const FGeometry& MyGeometry, const FPointerEvent&
 
 FVector2D STimeline::ComputeDesiredSize( float ) const
 {
-	return FVector2D(500.0f, 50.0f);
+	return FVector2D( 8.0f, 8.0f );
 }
