Fix GPU profiler not working in -game mode.  Made task graph have its own style

https://jira.ol.epicgames.net/browse/UE-7607

diff --git a/Engine/Source/Developer/TaskGraph/Private/SBarVisualizer.cpp b/Engine/Source/Developer/TaskGraph/Private/SBarVisualizer.cpp
--- a/Engine/Source/Developer/TaskGraph/Private/SBarVisualizer.cpp
+++ b/Engine/Source/Developer/TaskGraph/Private/SBarVisualizer.cpp
@@ -1,12 +1,12 @@
 // Copyright 1998-2015 Epic Games, Inc. All Rights Reserved.
 
 #include "SlateBasics.h"
-#include "EditorStyle.h"
 #include "TaskGraphInterfaces.h"
 #include "VisualizerEvents.h"
 #include "SGraphBar.h"
 #include "SBarVisualizer.h"
 #include "STimeline.h"
+#include "TaskGraphStyle.h"
 
 BEGIN_SLATE_FUNCTION_BUILD_OPTIMIZATION
 void SBarVisualizer::Construct( const FArguments& InArgs )
@@ -43,8 +43,8 @@ void SBarVisualizer::Construct( const FArguments& InArgs )
 		ViewMenuBuilder.AddMenuEntry( NSLOCTEXT("SBarVisualizer", "Flat", "Flat"), FText(), FSlateIcon(), Action, NAME_None, EUserInterfaceActionType::Check );
 	}
 
-	const FSlateBrush* HomeButtonBrush = FEditorStyle::GetBrush( "TaskGraph.Home" );
-	const FSlateBrush* ToParentButtonBrush = FEditorStyle::GetBrush( "TaskGraph.ToParent" );
+	const FSlateBrush* HomeButtonBrush = FTaskGraphStyle::Get()->GetBrush( "TaskGraph.Home" );
+	const FSlateBrush* ToParentButtonBrush = FTaskGraphStyle::Get()->GetBrush( "TaskGraph.ToParent" );
 
 	this->ChildSlot
 	[
@@ -52,14 +52,14 @@ void SBarVisualizer::Construct( const FArguments& InArgs )
 		+SVerticalBox::Slot().AutoHeight() .Padding( 2 ) .VAlign( VAlign_Fill )
 		[
 			SNew( SBorder )
-			.BorderImage( FEditorStyle::GetBrush("StatsHeader") )
-			.ForegroundColor( FEditorStyle::GetSlateColor("DefaultForeground") )
+			.BorderImage( FTaskGraphStyle::Get()->GetBrush("StatsHeader") )
+			.ForegroundColor( FTaskGraphStyle::Get()->GetSlateColor("DefaultForeground") )
 			[
 				SNew(SHorizontalBox)
 				+SHorizontalBox::Slot().AutoWidth() .Padding( 2 ) .HAlign( HAlign_Left )
 				[
 					SNew(SButton)
-					.ButtonStyle( FEditorStyle::Get(), "NoBorder" )
+					.ButtonStyle( FCoreStyle::Get(), "NoBorder" )
 					.ForegroundColor( FSlateColor::UseForeground() )
 					.ContentPadding(FMargin(0))
 					.Visibility( this, &SBarVisualizer::GetToParentButtonVisibility )
@@ -72,7 +72,7 @@ void SBarVisualizer::Construct( const FArguments& InArgs )
 				+SHorizontalBox::Slot().AutoWidth().Padding( 2 ).HAlign( HAlign_Left )
 				[
 					SNew(SButton)
-					.ButtonStyle( FEditorStyle::Get(), "NoBorder" )
+					.ButtonStyle( FCoreStyle::Get(), "NoBorder" )
 					.HAlign(HAlign_Center)
 					.VAlign(VAlign_Center)
 					.ForegroundColor( FSlateColor::UseForeground() )
@@ -94,12 +94,12 @@ void SBarVisualizer::Construct( const FArguments& InArgs )
 					SNew( SComboButton )
 					//.ToolTipText(NSLOCTEXT("PropertyEditor", "ResetToDefaultToolTip", "Reset to Default").ToString())
 					.HasDownArrow( false )
-					.ButtonStyle( FEditorStyle::Get(), "NoBorder" )
+					.ButtonStyle( FCoreStyle::Get(), "NoBorder" )
 					.ContentPadding(0)
 					.ButtonContent()
 					[
 						SNew(SImage)
-						.Image( FEditorStyle::GetBrush("EditorViewportToolBar.MenuDropdown") )
+						.Image( FTaskGraphStyle::Get()->GetBrush("TaskGraph.MenuDropdown") )
 					]
 					.MenuContent()
 					[
@@ -378,19 +378,19 @@ TSharedRef<ITableRow> SBarVisualizer::OnGenerateWidgetForList( TSharedPtr<FVisua
 			.AutoWidth()
 			[
 				SNew(SButton)
-				.ButtonStyle( FEditorStyle::Get(), "NoBorder" )
+				.ButtonStyle( FCoreStyle::Get(), "NoBorder" )
 				.ForegroundColor( FSlateColor::UseForeground() )
 				.ContentPadding(FMargin(0))
 				.OnClicked( this, &SBarVisualizer::ExpandBar, InItem )
 				[
 					SNew(SBorder) 
-					.BorderImage( FEditorStyle::GetBrush("NoBorder") )
+					.BorderImage( FCoreStyle::Get().GetBrush("NoBorder") )
 					.HAlign(HAlign_Center)
 					.VAlign(VAlign_Center)
 					.Padding(0)						
 					[
 						SNew(SImage)
-						.Image( FEditorStyle::GetBrush( "TreeArrow_Collapsed" ) )
+						.Image( FCoreStyle::Get().GetBrush( "TreeArrow_Collapsed" ) )
 					]
 				]
 			]
@@ -409,7 +409,7 @@ TSharedRef<ITableRow> SBarVisualizer::OnGenerateWidgetForList( TSharedPtr<FVisua
 		[
 			SNew( SBorder )
 			.Padding(0)
-			.BorderImage( FEditorStyle::GetBrush("NoBorder") )
+			.BorderImage( FCoreStyle::Get().GetBrush("NoBorder") )
 			// Handle right-click event for context menu
 			.OnMouseButtonDown( this, &SBarVisualizer::OnBarRightClicked, InItem )
 			[
diff --git a/Engine/Source/Developer/TaskGraph/Private/SGraphBar.cpp b/Engine/Source/Developer/TaskGraph/Private/SGraphBar.cpp
--- a/Engine/Source/Developer/TaskGraph/Private/SGraphBar.cpp
+++ b/Engine/Source/Developer/TaskGraph/Private/SGraphBar.cpp
@@ -1,7 +1,7 @@
 // Copyright 1998-2015 Epic Games, Inc. All Rights Reserved.
 
 #include "SlateBasics.h"
-#include "EditorStyle.h"
+#include "TaskGraphStyle.h"
 #include "TaskGraphInterfaces.h"
 #include "VisualizerEvents.h"
 #include "SGraphBar.h"
@@ -44,9 +44,9 @@ void SGraphBar::Construct( const FArguments& InArgs )
 	OnSelectionChanged = InArgs._OnSelectionChanged;
 	OnGeometryChanged = InArgs._OnGeometryChanged;
 
-	BackgroundImage = FEditorStyle::GetBrush("TaskGraph.Background");
-	FillImage = FEditorStyle::GetBrush("TaskGraph.Mono");
-	SelectedImage = FEditorStyle::GetBrush("TaskGraph.Selected");
+	BackgroundImage = FTaskGraphStyle::Get()->GetBrush("TaskGraph.Background");
+	FillImage = FTaskGraphStyle::Get()->GetBrush("TaskGraph.Mono");
+	SelectedImage = FTaskGraphStyle::Get()->GetBrush("TaskGraph.Selected");
 
 	LastHoveredEvent = INDEX_NONE;
 	Zoom = 1.0f;
@@ -67,7 +67,7 @@ int32 SGraphBar::OnPaint( const FPaintArgs& Args, const FGeometry& AllottedGeome
 	const FColor SelectedBarColor( 255, 255, 255, 255 );
 
 	// Paint inside the border only. 
-	const FVector2D BorderPadding = FEditorStyle::GetVector("ProgressBar.BorderPadding");
+	const FVector2D BorderPadding = FTaskGraphStyle::Get()->GetVector("TaskGraph.ProgressBar.BorderPadding");
 	const FSlateRect ForegroundClippingRect = AllottedGeometry.GetClippingRect().InsetBy(FMargin(BorderPadding.X, BorderPadding.Y)).IntersectionWith(MyClippingRect);
 
 	FSlateDrawElement::MakeBox(
diff --git a/Engine/Source/Developer/TaskGraph/Private/SProfileVisualizer.cpp b/Engine/Source/Developer/TaskGraph/Private/SProfileVisualizer.cpp
--- a/Engine/Source/Developer/TaskGraph/Private/SProfileVisualizer.cpp
+++ b/Engine/Source/Developer/TaskGraph/Private/SProfileVisualizer.cpp
@@ -16,8 +16,7 @@ void SProfileVisualizer::Construct( const FArguments& InArgs )
 	HeaderMessageText = InArgs._HeaderMessageText;
 	HeaderMessageTextColor = InArgs._HeaderMessageTextColor;
 
-	const FSlateBrush* ContentAreaBrush = FEditorStyle::GetBrush( "Docking.Tab", ".ContentAreaBrush" );
-
+	const FSlateBrush* ContentAreaBrush = FTaskGraphStyle::Get()->GetBrush("TaskGraph.ContentAreaBrush");
 	ChildSlot
 	[
 		SNew(SVerticalBox)
diff --git a/Engine/Source/Developer/TaskGraph/Private/STaskGraph.cpp b/Engine/Source/Developer/TaskGraph/Private/STaskGraph.cpp
--- a/Engine/Source/Developer/TaskGraph/Private/STaskGraph.cpp
+++ b/Engine/Source/Developer/TaskGraph/Private/STaskGraph.cpp
@@ -8,6 +8,7 @@
 #include "Tickable.h"
 #include "SProfileVisualizer.h"
 #include "SDockTab.h"
+#include "TaskGraphStyle.h"
 
 /**
  * Creates Visualizer using Visualizer profile data format
@@ -114,6 +115,7 @@ static TSharedPtr< FDelayedVisualizerSpawner > GDelayedVisualizerSpawner;
 
 void InitProfileVisualizer()
 {
+	FTaskGraphStyle::Initialize();
 	if( GDelayedVisualizerSpawner.IsValid() == false )
 	{
 		GDelayedVisualizerSpawner = MakeShareable( new FDelayedVisualizerSpawner() );
@@ -122,6 +124,7 @@ void InitProfileVisualizer()
 
 void ShutdownProfileVisualizer()
 {
+	FTaskGraphStyle::Shutdown();
 	GDelayedVisualizerSpawner.Reset();
 }
 
diff --git a/Engine/Source/Developer/TaskGraph/Private/STimeline.cpp b/Engine/Source/Developer/TaskGraph/Private/STimeline.cpp
--- a/Engine/Source/Developer/TaskGraph/Private/STimeline.cpp
+++ b/Engine/Source/Developer/TaskGraph/Private/STimeline.cpp
@@ -1,7 +1,7 @@
 // Copyright 1998-2015 Epic Games, Inc. All Rights Reserved.
 
 #include "SlateBasics.h"
-#include "EditorStyle.h"
+#include "TaskGraphStyle.h"
 #include "TaskGraphInterfaces.h"
 #include "VisualizerEvents.h"
 #include "STimeline.h"
@@ -12,7 +12,7 @@ void STimeline::Construct( const FArguments& InArgs )
 	MaxValue = InArgs._MaxValue;
 	FixedLabelSpacing = InArgs._FixedLabelSpacing;
 
-	BackgroundImage = FEditorStyle::GetBrush("TaskGraph.Background");
+	BackgroundImage = FTaskGraphStyle::Get()->GetBrush("TaskGraph.Background");
 
 	Zoom = 1.0f;
 	Offset = 0.0f;
@@ -31,7 +31,7 @@ int32 STimeline::OnPaint( const FPaintArgs& Args, const FGeometry& AllottedGeome
 	const FColor SelectedBarColor( 255, 255, 255, 255 );
 
 	// Paint inside the border only. 
-	const FVector2D BorderPadding = FEditorStyle::GetVector("ProgressBar.BorderPadding");
+	const FVector2D BorderPadding = FTaskGraphStyle::Get()->GetVector("TaskGraph.ProgressBar.BorderPadding");
 	const FSlateRect ForegroundClippingRect = AllottedGeometry.GetClippingRect().InsetBy(FMargin(BorderPadding.X, BorderPadding.Y)).IntersectionWith(MyClippingRect);
 
 	const float OffsetX = DrawingOffsetX; // BorderPadding.X
diff --git a/Engine/Source/Developer/TaskGraph/Private/TaskGraphStyle.cpp b/Engine/Source/Developer/TaskGraph/Private/TaskGraphStyle.cpp
new file mode 100644
--- /dev/null
+++ b/Engine/Source/Developer/TaskGraph/Private/TaskGraphStyle.cpp
@@ -0,0 +1,62 @@
+// Copyright 1998-2015 Epic Games, Inc. All Rights Reserved.
+
+#include "SlateBasics.h"
+#include "TaskGraphStyle.h"
+
+#define IMAGE_BRUSH( RelativePath, ... ) FSlateImageBrush( RootToContentDir( RelativePath, TEXT(".png") ), __VA_ARGS__ )
+#define BOX_BRUSH( RelativePath, ... ) FSlateBoxBrush( RootToContentDir( RelativePath, TEXT(".png") ), __VA_ARGS__ )
+#define BORDER_BRUSH( RelativePath, ... ) FSlateBorderBrush( RootToContentDir( RelativePath, TEXT(".png") ), __VA_ARGS__ )
+
+TSharedPtr< FTaskGraphStyle::FStyle > FTaskGraphStyle::StyleInstance;
+
+FTaskGraphStyle::FStyle::FStyle()
+	: FSlateStyleSet("TaskGraphStyle")
+{
+
+}
+
+void FTaskGraphStyle::FStyle::Initialize()
+{
+	
+
+	StyleInstance->SetContentRoot( FPaths::EngineContentDir() / TEXT("Editor/Slate") );
+	StyleInstance->SetCoreContentRoot(FPaths::EngineContentDir() / TEXT("Slate"));
+
+	StyleInstance->Set( "StatsHeader", new BOX_BRUSH( "Common/TableViewHeader", 4.f/32.f ) );
+	StyleInstance->Set( "TaskGraph.ContentAreaBrush", new BOX_BRUSH( "/Docking/TabContentArea", FMargin(4/16.0f) ) );
+	StyleInstance->Set( "TaskGraph.Background", new BOX_BRUSH( "Common/ProgressBar_Background", FMargin(5.f/12.f) ) );
+	StyleInstance->Set( "TaskGraph.Normal", new BOX_BRUSH( "Common/TaskGraph_Normal", FMargin(5.f/12.f) ) );
+	StyleInstance->Set( "TaskGraph.Selected", new BOX_BRUSH( "Common/TaskGraph_Selected", FMargin(5.f/12.f) ) );
+	StyleInstance->Set( "TaskGraph.Mono", new BOX_BRUSH( "Common/TaskGraph_Mono", FMargin(5.f/12.f) ) );
+	StyleInstance->Set( "TaskGraph.BorderPadding", FVector2D(1,0) );
+	StyleInstance->Set( "TaskGraph.SortUp", new IMAGE_BRUSH( "Common/SortUpArrow", FVector2D(8,4) ) );
+	StyleInstance->Set( "TaskGraph.SortDown", new IMAGE_BRUSH( "Common/SortDownArrow", FVector2D(8,4) ) );
+	StyleInstance->Set( "TaskGraph.Home", new IMAGE_BRUSH( "Icons/Home16x16", FVector2D(16,16) ) );
+	StyleInstance->Set( "TaskGraph.ToParent", new IMAGE_BRUSH( "Icons/ToParent", FVector2D(16,16) ) );
+	StyleInstance->Set( "TaskGraph.ProgressBar.BorderPadding", FVector2D(1,0) );
+	StyleInstance->Set( "TaskGraph.MenuDropdown", new IMAGE_BRUSH( "Common/ComboArrow", FVector2D(8,8) ) );
+
+	FSlateStyleRegistry::RegisterSlateStyle(*StyleInstance.Get());
+}
+
+
+void FTaskGraphStyle::Initialize() 
+{
+	StyleInstance = MakeShareable( new FTaskGraphStyle::FStyle );
+	StyleInstance->Initialize();
+}
+
+void FTaskGraphStyle::Shutdown()
+{
+	if (StyleInstance.IsValid())
+	{
+		FSlateStyleRegistry::UnRegisterSlateStyle(*StyleInstance.Get());
+		ensure(StyleInstance.IsUnique());
+		StyleInstance.Reset();
+	}
+}
+
+TSharedPtr< ISlateStyle > FTaskGraphStyle::Get()
+{
+	return StyleInstance;
+}
diff --git a/Engine/Source/Developer/TaskGraph/Private/TaskGraphStyle.h b/Engine/Source/Developer/TaskGraph/Private/TaskGraphStyle.h
new file mode 100644
--- /dev/null
+++ b/Engine/Source/Developer/TaskGraph/Private/TaskGraphStyle.h
@@ -0,0 +1,30 @@
+// Copyright 1998-2015 Epic Games, Inc. All Rights Reserved.
+
+#pragma once
+
+
+#include "EditorStyle.h"
+
+/**
+ * Declares the Editor's visual style.
+ */
+class FTaskGraphStyle
+{
+public:
+	
+	class FStyle : public FSlateStyleSet
+	{
+	public:
+		FStyle();
+
+		void Initialize();
+	};
+
+	static void Initialize();
+
+	static void Shutdown();
+
+	static TSharedPtr<class ISlateStyle> Get();
+
+	static TSharedPtr< FTaskGraphStyle::FStyle > StyleInstance;
+};
diff --git a/Engine/Source/Editor/EditorStyle/Private/SlateEditorStyle.cpp b/Engine/Source/Editor/EditorStyle/Private/SlateEditorStyle.cpp
--- a/Engine/Source/Editor/EditorStyle/Private/SlateEditorStyle.cpp
+++ b/Engine/Source/Editor/EditorStyle/Private/SlateEditorStyle.cpp
@@ -635,19 +635,6 @@ void FSlateEditorStyle::FStyle::SetupGeneralStyles()
 			);
 	}
 
-	// TaskGraph
-	{
-		Set( "TaskGraph.Background", new BOX_BRUSH( "Common/ProgressBar_Background", FMargin(5.f/12.f) ) );
-		Set( "TaskGraph.Normal", new BOX_BRUSH( "Common/TaskGraph_Normal", FMargin(5.f/12.f) ) );
-		Set( "TaskGraph.Selected", new BOX_BRUSH( "Common/TaskGraph_Selected", FMargin(5.f/12.f) ) );
-		Set( "TaskGraph.Mono", new BOX_BRUSH( "Common/TaskGraph_Mono", FMargin(5.f/12.f) ) );
-		Set( "TaskGraph.BorderPadding", FVector2D(1,0) );
-		Set( "TaskGraph.SortUp", new IMAGE_BRUSH( "Common/SortUpArrow", Icon8x4 ) );
-		Set( "TaskGraph.SortDown", new IMAGE_BRUSH( "Common/SortDownArrow", Icon8x4 ) );
-		Set( "TaskGraph.Home", new IMAGE_BRUSH( "Icons/Home16x16", Icon16x16 ) );
-		Set( "TaskGraph.ToParent", new IMAGE_BRUSH( "Icons/ToParent", Icon16x16 ) );
-	}
-
 	// Tool panels
 	{
 		Set( "ToolPanel.GroupBorder", new BOX_BRUSH( "Common/GroupBorder", FMargin(4.0f/16.0f) ) );
@@ -1820,10 +1807,7 @@ void FSlateEditorStyle::FStyle::SetupGeneralStyles()
 		);
 	}
 
-	// Stats visualizer
-	{
-		Set( "StatsHeader", new BOX_BRUSH( "Common/TableViewHeader", 4.f/32.f ) );
-	}
+
 
 	// Package Migration
 	{
