Fix non-unity

diff --git a/Engine/Source/Developer/TaskGraph/Private/SProfileVisualizer.cpp b/Engine/Source/Developer/TaskGraph/Private/SProfileVisualizer.cpp
--- a/Engine/Source/Developer/TaskGraph/Private/SProfileVisualizer.cpp
+++ b/Engine/Source/Developer/TaskGraph/Private/SProfileVisualizer.cpp
@@ -1,7 +1,7 @@
 // Copyright 1998-2015 Epic Games, Inc. All Rights Reserved.
 
 #include "SlateBasics.h"
-#include "EditorStyle.h"
+#include "STaskGraphStyle.h"
 #include "SProfileVisualizer.h"
 #include "TaskGraphInterfaces.h"
 #include "STaskGraph.h"
