Prevent TaskGraph style from looking for style resources in cooked builds

diff --git a/Engine/Source/Developer/LogVisualizer/LogVisualizer.Build.cs b/Engine/Source/Developer/LogVisualizer/LogVisualizer.Build.cs
--- a/Engine/Source/Developer/LogVisualizer/LogVisualizer.Build.cs
+++ b/Engine/Source/Developer/LogVisualizer/LogVisualizer.Build.cs
@@ -16,9 +16,7 @@ public class LogVisualizer : ModuleRules
 
 		PublicIncludePaths.AddRange(
 			new string[] {
-				"Developer/TaskGraph/Public",
 				"Runtime/Engine/Classes",
-                "Editor/GeometryMode/Public",
                 "Editor/WorkspaceMenuStructure/Public"
 			}
 		);
@@ -33,7 +31,6 @@ public class LogVisualizer : ModuleRules
 				"SlateCore",
                 "EditorStyle",
 				"Engine",
-				"TaskGraph",
 				"UnrealEd",
                 "SequencerWidgets",
 			}
diff --git a/Engine/Source/Developer/LogVisualizer/Private/LogVisualizerSessionSettings.cpp b/Engine/Source/Developer/LogVisualizer/Private/LogVisualizerSessionSettings.cpp
--- a/Engine/Source/Developer/LogVisualizer/Private/LogVisualizerSessionSettings.cpp
+++ b/Engine/Source/Developer/LogVisualizer/Private/LogVisualizerSessionSettings.cpp
@@ -6,7 +6,6 @@
 #if WITH_EDITOR
 #include "Editor/EditorEngine.h"
 #include "ISettingsModule.h"
-#include "GeometryEdMode.h"
 #include "UnrealEdMisc.h"
 #endif // WITH_EDITOR
 
diff --git a/Engine/Source/Developer/LogVisualizer/Private/LogVisualizerSettings.cpp b/Engine/Source/Developer/LogVisualizer/Private/LogVisualizerSettings.cpp
--- a/Engine/Source/Developer/LogVisualizer/Private/LogVisualizerSettings.cpp
+++ b/Engine/Source/Developer/LogVisualizer/Private/LogVisualizerSettings.cpp
@@ -6,7 +6,6 @@
 #if WITH_EDITOR
 #include "Editor/EditorEngine.h"
 #include "ISettingsModule.h"
-#include "GeometryEdMode.h"
 #include "UnrealEdMisc.h"
 #endif // WITH_EDITOR
 
diff --git a/Engine/Source/Developer/TaskGraph/Private/TaskGraphStyle.cpp b/Engine/Source/Developer/TaskGraph/Private/TaskGraphStyle.cpp
--- a/Engine/Source/Developer/TaskGraph/Private/TaskGraphStyle.cpp
+++ b/Engine/Source/Developer/TaskGraph/Private/TaskGraphStyle.cpp
@@ -17,8 +17,7 @@ FTaskGraphStyle::FStyle::FStyle()
 
 void FTaskGraphStyle::FStyle::Initialize()
 {
-	
-
+#if WITH_EDITORONLY_DATA
 	StyleInstance->SetContentRoot( FPaths::EngineContentDir() / TEXT("Editor/Slate") );
 	StyleInstance->SetCoreContentRoot(FPaths::EngineContentDir() / TEXT("Slate"));
 
@@ -37,6 +36,7 @@ void FTaskGraphStyle::FStyle::Initialize()
 	StyleInstance->Set( "TaskGraph.MenuDropdown", new IMAGE_BRUSH( "Common/ComboArrow", FVector2D(8,8) ) );
 
 	FSlateStyleRegistry::RegisterSlateStyle(*StyleInstance.Get());
+#endif
 }
 
 
